Added missing copyright to fix CIS

diff --git a/Engine/Platforms/XXX/Source/Programs/UnrealBuildTool/UEBuildXXX.cs b/Engine/Platforms/XXX/Source/Programs/UnrealBuildTool/UEBuildXXX.cs
--- a/Engine/Platforms/XXX/Source/Programs/UnrealBuildTool/UEBuildXXX.cs
+++ b/Engine/Platforms/XXX/Source/Programs/UnrealBuildTool/UEBuildXXX.cs
@@ -1,3 +1,4 @@
+// Copyright 1998-2019 Epic Games, Inc. All Rights Reserved.
 using System;
 using System.Collections.Generic;
 using Tools.DotNETCommon;
diff --git a/Engine/Platforms/XXX/Source/Runtime/Core/Private/XXXPlatformMisc.cpp b/Engine/Platforms/XXX/Source/Runtime/Core/Private/XXXPlatformMisc.cpp
--- a/Engine/Platforms/XXX/Source/Runtime/Core/Private/XXXPlatformMisc.cpp
+++ b/Engine/Platforms/XXX/Source/Runtime/Core/Private/XXXPlatformMisc.cpp
@@ -1 +1,2 @@
+// Copyright 1998-2019 Epic Games, Inc. All Rights Reserved.
 foo
