- Fixed XXX on Mac

diff --git a/Engine/Platforms/XXX/Source/Programs/UnrealBuildTool/UEBuildXXX.cs b/Engine/Platforms/XXX/Source/Programs/UnrealBuildTool/UEBuildXXX.cs
--- a/Engine/Platforms/XXX/Source/Programs/UnrealBuildTool/UEBuildXXX.cs
+++ b/Engine/Platforms/XXX/Source/Programs/UnrealBuildTool/UEBuildXXX.cs
@@ -199,7 +199,7 @@ class XXXPlatform : UEBuildPlatform
 			Rules.PublicDefinitions.Add("PLATFORM_XXX=0");
 		}
 
-		if (Target.Platform == UnrealTargetPlatform.Win64)
+		if (Target.Platform == UnrealTargetPlatform.Win64 || Target.Platform == UnrealTargetPlatform.Mac || Target.Platform == UnrealTargetPlatform.Linux)
 		{
 			if (ModuleName == "RHI")
 			{
