Fix paths to platform icons for Fake and DPX platforms.

diff --git a/Engine/Platforms/XXX/Config/DataDrivenPlatformInfo.ini b/Engine/Platforms/XXX/Config/DataDrivenPlatformInfo.ini
--- a/Engine/Platforms/XXX/Config/DataDrivenPlatformInfo.ini
+++ b/Engine/Platforms/XXX/Config/DataDrivenPlatformInfo.ini
@@ -21,8 +21,8 @@ TargetPlatformName=XXX
 DisplayName=FakePlatformXXX
 PlatformType=Game
 PlatformFlags=None
-NormalIconPath=/Platforms/XXX/Engine/Content/Launcher/Platform_XXX_24x
-LargeIconPath=/Platforms/XXX/Engine/Content/Launcher/Platform_XXX_128x
+NormalIconPath=/Platforms/XXX/Content/Launcher/Platform_XXX_24x
+LargeIconPath=/Platforms/XXX/Content/Launcher/Platform_XXX_128x
 XLargeIconPath=
 UATCommandLine=
 AutoSDKPath=
