Change platform extensions and games to output UAT scripts to their own folder rather than to the Engine folder.

diff --git a/Engine/Platforms/XXX/Source/Programs/AutomationTool/XXX.Automation.csproj b/Engine/Platforms/XXX/Source/Programs/AutomationTool/XXX.Automation.csproj
--- a/Engine/Platforms/XXX/Source/Programs/AutomationTool/XXX.Automation.csproj
+++ b/Engine/Platforms/XXX/Source/Programs/AutomationTool/XXX.Automation.csproj
@@ -18,7 +18,7 @@
     <DebugSymbols>true</DebugSymbols>
     <DebugType>full</DebugType>
     <Optimize>false</Optimize>
-    <OutputPath>..\..\..\..\..\..\Engine\Binaries\DotNET\AutomationScripts\XXX\</OutputPath>
+    <OutputPath>..\..\..\Binaries\DotNET\AutomationScripts\</OutputPath>
     <DefineConstants>DEBUG;TRACE</DefineConstants>
     <ErrorReport>prompt</ErrorReport>
     <WarningLevel>4</WarningLevel>
@@ -28,7 +28,7 @@
   <PropertyGroup Condition=" '$(Configuration)|$(Platform)' == 'Development|AnyCPU' ">
     <DebugType>pdbonly</DebugType>
     <Optimize>false</Optimize>
-    <OutputPath>..\..\..\..\..\..\Engine\Binaries\DotNET\AutomationScripts\XXX\</OutputPath>
+    <OutputPath>..\..\..\Binaries\DotNET\AutomationScripts\</OutputPath>
     <DefineConstants>TRACE</DefineConstants>
     <ErrorReport>prompt</ErrorReport>
     <WarningLevel>4</WarningLevel>
